Config Editor added to UE4.
We now have the capacity to set a uproperty value anywhere in the config file hierarchy

Satisfies UEPLAT-429

The user can now use the 'ConfigHierarchyEditable' meta flag in their uproperties to present a butoon that allows the property to be edited in the config hierarchy, where applicable.
Target platform is configurable.

diff --git a/Engine/Source/Editor/ConfigEditor/ConfigEditor.Build.cs b/Engine/Source/Editor/ConfigEditor/ConfigEditor.Build.cs
new file mode 100644
--- /dev/null
+++ b/Engine/Source/Editor/ConfigEditor/ConfigEditor.Build.cs
@@ -0,0 +1,57 @@
+// Copyright 1998-2015 Epic Games, Inc. All Rights Reserved.
+
+using UnrealBuildTool;
+
+public class ConfigEditor : ModuleRules
+{
+	public ConfigEditor(TargetInfo Target)
+	{
+		PrivateIncludePaths.AddRange(
+			new string[] {
+				"Editor/ConfigEditor/Private",
+				"Editor/ConfigEditor/Private/PropertyVisualization",
+			}
+		);
+		
+		PublicIncludePaths.AddRange(
+			new string[] {
+				"Editor/ConfigEditor/Public",
+				"Editor/ConfigEditor/Public/PropertyVisualization",
+			}
+		);
+	
+		PrivateDependencyModuleNames.AddRange(
+			new string[] {
+				"Core",
+				"CoreUObject",
+				"EditorStyle",
+				"InputCore",
+				"PropertyEditor",
+				"Slate",
+				"SlateCore",
+				"SourceControl",
+				"TargetPlatform",
+				"WorkspaceMenuStructure",
+			}
+		);
+
+
+		CircularlyReferencedDependentModules.AddRange(
+			new string[] 
+			{
+				"PropertyEditor",
+			}
+		); 
+
+
+		PrivateIncludePathModuleNames.AddRange(
+			new string[] {
+			}
+		);
+
+		DynamicallyLoadedModuleNames.AddRange(
+			new string[] {
+			}
+		);
+	}
+}
diff --git a/Engine/Source/Editor/ConfigEditor/Private/ConfigEditorModule.cpp b/Engine/Source/Editor/ConfigEditor/Private/ConfigEditorModule.cpp
new file mode 100644
--- /dev/null
+++ b/Engine/Source/Editor/ConfigEditor/Private/ConfigEditorModule.cpp
@@ -0,0 +1,65 @@
+// Copyright 1998-2015 Epic Games, Inc. All Rights Reserved.
+
+
+#include "ConfigEditorPCH.h"
+#include "ConfigEditorModule.h"
+
+#include "SDockTab.h"
+
+DEFINE_LOG_CATEGORY_STATIC(ConfigEditor, Log, All);
+
+/*-----------------------------------------------------------------------------
+   FConfigEditorModule
+-----------------------------------------------------------------------------*/
+
+namespace ConfigEditorModule
+{
+	static const FName ConfigEditorId = FName(TEXT("ConfigEditor"));
+}
+
+
+void FConfigEditorModule::StartupModule()
+{
+	FGlobalTabmanager::Get()->RegisterNomadTabSpawner(ConfigEditorModule::ConfigEditorId, FOnSpawnTab::CreateRaw(this, &FConfigEditorModule::SpawnConfigEditorTab))
+		.SetDisplayName(NSLOCTEXT("ConfigEditorModule", "TabTitle", "Config Editor"))
+		.SetTooltipText(NSLOCTEXT("ConfigEditorModule", "TooltipText", "Open the Config Editor tab."))
+		.SetIcon(FSlateIcon(FEditorStyle::GetStyleSetName(), "ConfigEditor.TabIcon"));
+}
+
+
+void FConfigEditorModule::ShutdownModule()
+{
+
+}
+
+
+void FConfigEditorModule::AddExternalPropertyValueWidgetAndConfigPairing(const FString& ConfigFile, const TSharedPtr<SWidget> ValueWidget)
+{
+	ExternalPropertyValueWidgetAndConfigPairings.Add(ConfigFile, ValueWidget);
+}
+
+
+TSharedRef<SWidget> FConfigEditorModule::GetValueWidgetForConfigProperty(const FString& ConfigFile)
+{
+	TSharedPtr<SWidget>* ValueWidget = ExternalPropertyValueWidgetAndConfigPairings.Find(ConfigFile);
+	return (ValueWidget != nullptr && (*ValueWidget).IsValid()) ? (*ValueWidget).ToSharedRef() : SNullWidget::NullWidget;
+}
+
+
+TSharedRef<SDockTab> FConfigEditorModule::SpawnConfigEditorTab(const FSpawnTabArgs& Args)
+{
+	return SNew(SDockTab)
+		.TabRole(ETabRole::NomadTab)
+		[
+			SAssignNew(PropertyConfigEditor, SConfigEditor, CachedPropertyToView)
+		];
+}
+
+
+void FConfigEditorModule::CreateHierarchyEditor(UProperty* InEditProperty)
+{
+	CachedPropertyToView = InEditProperty;
+}
+
+
+IMPLEMENT_MODULE(FConfigEditorModule, ConfigEditor);
diff --git a/Engine/Source/Editor/ConfigEditor/Private/ConfigEditorModule.h b/Engine/Source/Editor/ConfigEditor/Private/ConfigEditorModule.h
new file mode 100644
--- /dev/null
+++ b/Engine/Source/Editor/ConfigEditor/Private/ConfigEditorModule.h
@@ -0,0 +1,50 @@
+// Copyright 1998-2015 Epic Games, Inc. All Rights Reserved.
+
+#pragma once
+
+// Module includes
+#include "IConfigEditorModule.h"
+#include "SConfigEditor.h"
+
+/*-----------------------------------------------------------------------------
+   FConfigEditorModule
+-----------------------------------------------------------------------------*/
+class FConfigEditorModule : public IConfigEditorModule
+{
+public:
+	// Begin IModuleInterface interface
+	virtual void StartupModule() override;
+	virtual void ShutdownModule() override;
+	// End IModuleInterface interface
+
+
+public:
+	// Begin IConfigEditorModule interface
+	virtual void CreateHierarchyEditor(UProperty* EditProperty) override;
+	virtual void AddExternalPropertyValueWidgetAndConfigPairing(const FString& ConfigFile, const TSharedPtr<SWidget> ValueWidget) override;
+	virtual TSharedRef<SWidget> GetValueWidgetForConfigProperty(const FString& ConfigFile) override;
+	// End IConfigEditorModule interface	
+	
+
+private:
+	/** 
+	 * Creates a new Config editor tab 
+	 *
+	 * @return A spawned tab containing a config editor for a pre-defined property
+	 */
+	TSharedRef<class SDockTab> SpawnConfigEditorTab(const FSpawnTabArgs& Args);
+
+
+private:
+	// We use this to maintain a reference to a property value widget, ie. a combobox for enum, a check box for bool...
+	// and provide these references when they are needed later in the details view construction
+	// This is important.
+	TMap<FString, TSharedPtr<SWidget>> ExternalPropertyValueWidgetAndConfigPairings;
+
+	// Reference to the Config editor widget
+	TSharedPtr<SConfigEditor> PropertyConfigEditor;
+
+	// Reference to the property the hierarchy is to view.
+	TWeakObjectPtr<UProperty> CachedPropertyToView;
+};
+
diff --git a/Engine/Source/Editor/ConfigEditor/Private/ConfigEditorPCH.h b/Engine/Source/Editor/ConfigEditor/Private/ConfigEditorPCH.h
new file mode 100644
--- /dev/null
+++ b/Engine/Source/Editor/ConfigEditor/Private/ConfigEditorPCH.h
@@ -0,0 +1,8 @@
+// Copyright 1998-2015 Epic Games, Inc. All Rights Reserved.
+
+#pragma once
+
+// External include
+#include "SlateBasics.h"
+#include "SlateCore.h"
+#include "EditorStyleSet.h"
\ No newline at end of file
diff --git a/Engine/Source/Editor/ConfigEditor/Private/ConfigPropertyHelper.cpp b/Engine/Source/Editor/ConfigEditor/Private/ConfigPropertyHelper.cpp
new file mode 100644
--- /dev/null
+++ b/Engine/Source/Editor/ConfigEditor/Private/ConfigPropertyHelper.cpp
@@ -0,0 +1,43 @@
+// Copyright 1998-2015 Epic Games, Inc. All Rights Reserved.
+
+#include "ConfigEditorPCH.h"
+#include "ConfigPropertyHelper.h"
+#include "ISourceControlModule.h"
+
+void UPropertyConfigFileDisplayRow::InitWithConfigAndProperty(const FString& InConfigFileName, UProperty* InEditProperty)
+{
+	ConfigFileName = FPaths::ConvertRelativePathToFull(InConfigFileName);
+	ExternalProperty = InEditProperty;
+
+
+	ISourceControlProvider& SourceControlProvider = ISourceControlModule::Get().GetProvider();
+	// We will add source control soon...
+	FSourceControlStatePtr SourceControlState = nullptr; // SourceControlProvider.GetState(ConfigFileName, EStateCacheUsage::Use);
+
+	// Only include config files that are currently checked out or packages not under source control
+	{
+		if (FPaths::FileExists(ConfigFileName))
+		{
+			if (SourceControlState.IsValid())
+			{
+				bIsFileWritable = SourceControlState->IsCheckedOut() || SourceControlState->IsAdded();
+			}
+			else
+			{
+				bIsFileWritable = !IFileManager::Get().IsReadOnly(*ConfigFileName);
+			}
+		}
+		else
+		{
+			if (SourceControlState.IsValid())
+			{
+				bIsFileWritable = (SourceControlState->IsSourceControlled() && SourceControlState->CanAdd());
+			}
+			else
+			{
+				bIsFileWritable = false;
+			}
+		}
+	}
+
+}
\ No newline at end of file
diff --git a/Engine/Source/Editor/ConfigEditor/Private/PropertyVisualization/ConfigPropertyCellPresenter.cpp b/Engine/Source/Editor/ConfigEditor/Private/PropertyVisualization/ConfigPropertyCellPresenter.cpp
new file mode 100644
--- /dev/null
+++ b/Engine/Source/Editor/ConfigEditor/Private/PropertyVisualization/ConfigPropertyCellPresenter.cpp
@@ -0,0 +1,291 @@
+// Copyright 1998-2015 Epic Games, Inc. All Rights Reserved.
+
+#include "ConfigEditorPCH.h"
+#include "ConfigEditorModule.h"
+#include "ConfigPropertyCellPresenter.h"
+#include "ConfigPropertyHelper.h"
+
+#include "Editor/PropertyEditor/Public/PropertyEditorModule.h"
+#include "Editor/PropertyEditor/Public/PropertyHandle.h"
+#include "Editor/PropertyEditor/Public/PropertyPath.h"
+#include "Editor/PropertyEditor/Public/IPropertyTable.h"
+
+
+#define LOCTEXT_NAMESPACE "ConfigEditor"
+
+
+class SConfigPropertyCell
+	: public SCompoundWidget
+{
+public:
+	SLATE_USER_ARGS(SConfigPropertyCell) {}
+	SLATE_END_ARGS()
+
+	// Begin SCompoundWidget|SWidget interface
+	virtual void Construct(const FArguments& InArgs, UPropertyConfigFileDisplayRow* InPropertyOwnerDisplayObject);
+	virtual void OnMouseEnter(const FGeometry& MyGeometry, const FPointerEvent& MouseEvent) override;
+	virtual void OnMouseLeave(const FPointerEvent& MouseEvent) override;
+	// End SCompoundWidget|SWidget interface
+
+private:
+
+	// Create the displayable area for this cell.
+	void BuildDisplayAreaWidget();
+	// The shown display widget, includes all optional decoration
+	TSharedPtr<SWidget> DecoratedDisplayWidget;
+
+	// The control widget for editing the value of the property in the specified config.
+	TSharedPtr<SWidget> DisplayedValueWidget;
+	EVisibility GetDisplayedValueWidgetVisibility() const;
+	bool HandleDisplayedValueWidgetEnabled() const;
+
+	// The capacity to add an entry to the specified config file
+	TSharedPtr<SButton> AddPropertyToConfigButton;
+	FReply HandleAddPropertyToConfigClicked();
+	EVisibility GetAddPropertyToConfigButtonVisibility() const;
+
+	// WIP.. The capacity to remove an entry from the specified config file
+	TSharedPtr<SButton> RemovePropertyFromConfigButton;
+	FReply HandleRemovePropertyFromConfigClicked();
+	EVisibility GetRemovePropertyFromConfigVisibility() const;
+
+	// Check if the config file has an entry for this property
+	bool DoesConfigFileHaveEntryForProperty();
+	// Cached value of the above, as we use this for widget visibility callbacks.
+	bool bCachedConfigHasPropertyValue;
+
+	// The object which holds the property we are editing
+	UPropertyConfigFileDisplayRow* PropertyOwnerDisplayObject;
+
+};
+
+
+void SConfigPropertyCell::BuildDisplayAreaWidget()
+{
+	IConfigEditorModule& ConfigEditorModule = FModuleManager::Get().GetModuleChecked<IConfigEditorModule>("ConfigEditor");
+	DisplayedValueWidget = ConfigEditorModule.GetValueWidgetForConfigProperty(PropertyOwnerDisplayObject->ConfigFileName);
+	DisplayedValueWidget->SetEnabled(TAttribute<bool>::Create(TAttribute<bool>::FGetter::CreateSP(this, &SConfigPropertyCell::HandleDisplayedValueWidgetEnabled)));
+
+	static const FName DefaultForegroundName("DefaultForeground");
+
+	SAssignNew(DecoratedDisplayWidget, SHorizontalBox)
+		// The widget to alter the objects value in a config file
+		+ SHorizontalBox::Slot()
+		[
+			SNew(SVerticalBox)
+			.Visibility(this, &SConfigPropertyCell::GetDisplayedValueWidgetVisibility)
+			+ SVerticalBox::Slot()
+			.AutoHeight()
+			.VAlign(EVerticalAlignment::VAlign_Center)
+			[
+				SNew(SHorizontalBox)
+				+ SHorizontalBox::Slot()
+				.HAlign(HAlign_Center)
+				[
+					SNew(SBorder)
+					.BorderImage(FCoreStyle::Get().GetBrush("ToolPanel.GroupBorder"))
+					.ForegroundColor(FEditorStyle::GetSlateColor(DefaultForegroundName))
+					.Padding(0.0f)
+					[
+						DisplayedValueWidget.ToSharedRef()
+					]
+				]
+			]
+		]
+		// The add to config widget
+		+ SHorizontalBox::Slot()
+		.HAlign(HAlign_Center)
+		[
+			SNew(SVerticalBox)
+			.Visibility(this, &SConfigPropertyCell::GetAddPropertyToConfigButtonVisibility)
+			+ SVerticalBox::Slot()
+			.AutoHeight()
+			.VAlign(EVerticalAlignment::VAlign_Center)
+			[
+				SAssignNew(AddPropertyToConfigButton, SButton)
+				.OnClicked(this, &SConfigPropertyCell::HandleAddPropertyToConfigClicked)
+				.ButtonStyle(FEditorStyle::Get(), "HoverHintOnly")
+				[
+					SNew(SImage)
+					.Image(FEditorStyle::GetBrush(TEXT("Plus")))
+					.ColorAndOpacity(FSlateColor::UseForeground())
+				]
+			]
+		]
+		// The remove from config widget
+		+ SHorizontalBox::Slot()
+		.AutoWidth()
+		.HAlign(HAlign_Center)
+		[
+			SNew(SVerticalBox)
+			.Visibility(this, &SConfigPropertyCell::GetRemovePropertyFromConfigVisibility)
+			+ SVerticalBox::Slot()
+			.AutoHeight()
+			.VAlign(EVerticalAlignment::VAlign_Center)
+			.Padding(FMargin(4.0f, 0.0f))
+			[
+				SAssignNew(RemovePropertyFromConfigButton, SButton)
+				.OnClicked(this, &SConfigPropertyCell::HandleRemovePropertyFromConfigClicked)
+				.ButtonStyle(FEditorStyle::Get(), "HoverHintOnly")
+				[
+					SNew(SImage)
+					.Image(FEditorStyle::GetBrush(TEXT("Cross")))
+					.ColorAndOpacity(FSlateColor::UseForeground())
+				]
+			]
+		];
+}
+
+
+void SConfigPropertyCell::Construct(const FArguments& InArgs, UPropertyConfigFileDisplayRow* InPropertyOwnerDisplayObject)
+{
+	PropertyOwnerDisplayObject = InPropertyOwnerDisplayObject;
+	bCachedConfigHasPropertyValue = DoesConfigFileHaveEntryForProperty();
+
+	// Create the cell visualization
+	BuildDisplayAreaWidget();
+
+	ChildSlot
+	[
+		DecoratedDisplayWidget.ToSharedRef()
+	];
+}
+
+
+void SConfigPropertyCell::OnMouseEnter(const FGeometry& MyGeometry, const FPointerEvent& MouseEvent)
+{
+
+}
+
+
+void SConfigPropertyCell::OnMouseLeave(const FPointerEvent& MouseEvent)
+{
+
+}
+
+
+EVisibility SConfigPropertyCell::GetDisplayedValueWidgetVisibility() const
+{
+	return bCachedConfigHasPropertyValue ? EVisibility::Visible : EVisibility::Collapsed;
+}
+
+
+bool SConfigPropertyCell::HandleDisplayedValueWidgetEnabled() const
+{
+	return PropertyOwnerDisplayObject->bIsFileWritable;
+}
+
+
+FReply SConfigPropertyCell::HandleAddPropertyToConfigClicked()
+{
+	PropertyOwnerDisplayObject->UpdateDefaultConfigFile(PropertyOwnerDisplayObject->ConfigFileName);
+	bCachedConfigHasPropertyValue = true;
+
+	return FReply::Handled();
+}
+
+
+EVisibility SConfigPropertyCell::GetAddPropertyToConfigButtonVisibility() const
+{
+	return (!bCachedConfigHasPropertyValue && PropertyOwnerDisplayObject->bIsFileWritable) ? EVisibility::Visible : EVisibility::Collapsed;
+}
+
+
+FReply SConfigPropertyCell::HandleRemovePropertyFromConfigClicked()
+{
+	// This needs work...
+	UClass* OwnerClass = PropertyOwnerDisplayObject->ExternalProperty->GetOwnerClass();
+	FString SectionName = OwnerClass->GetPathName();
+	FString PropertyName = PropertyOwnerDisplayObject->ExternalProperty->GetName();
+
+	// This...
+	GConfig->RemoveKey(*SectionName, *PropertyName, PropertyOwnerDisplayObject->ConfigFileName);
+
+	// ... really oughta affect this...
+	PropertyOwnerDisplayObject->UpdateDefaultConfigFile(PropertyOwnerDisplayObject->ConfigFileName);
+
+	return FReply::Handled();
+}
+
+
+EVisibility SConfigPropertyCell::GetRemovePropertyFromConfigVisibility() const
+{
+	/* Technical limitations of UObjects and Config make this difficult. Hopefully we can add it at some point. */
+	return EVisibility::Collapsed;
+}
+
+
+bool SConfigPropertyCell::DoesConfigFileHaveEntryForProperty()
+{
+	UClass* OwnerClass = PropertyOwnerDisplayObject->ExternalProperty->GetOwnerClass();
+	FString SectionName = OwnerClass->GetPathName();
+	FString PropertyName = PropertyOwnerDisplayObject->ExternalProperty->GetName();
+	
+	FString TempValue;
+	return GConfig->GetString(*SectionName, *PropertyName, TempValue, PropertyOwnerDisplayObject->ConfigFileName);
+}
+
+
+
+FConfigPropertyCellPresenter::FConfigPropertyCellPresenter(const TSharedPtr< IPropertyHandle > PropertyHandle)
+{
+	TArray<UObject*> OuterObjects;
+	PropertyHandle->GetOuterObjects(OuterObjects);
+
+	if (OuterObjects.Num() == 1)
+	{
+		UPropertyConfigFileDisplayRow* PropertyOwner = CastChecked<UPropertyConfigFileDisplayRow>(OuterObjects[0]);
+		DisplayWidget = SNew(SConfigPropertyCell, PropertyOwner);
+	}
+}
+
+
+TSharedRef< class SWidget > FConfigPropertyCellPresenter::ConstructDisplayWidget()
+{
+	return DisplayWidget.ToSharedRef();
+}
+
+
+bool FConfigPropertyCellPresenter::RequiresDropDown()
+{
+	return false;
+}
+
+
+TSharedRef< class SWidget > FConfigPropertyCellPresenter::ConstructEditModeCellWidget()
+{
+	return ConstructDisplayWidget();
+}
+
+
+TSharedRef< class SWidget > FConfigPropertyCellPresenter::ConstructEditModeDropDownWidget()
+{
+	return SNullWidget::NullWidget;
+}
+
+
+TSharedRef< class SWidget > FConfigPropertyCellPresenter::WidgetToFocusOnEdit()
+{
+	return SNullWidget::NullWidget;
+}
+
+
+bool FConfigPropertyCellPresenter::HasReadOnlyEditMode()
+{
+	return true;
+}
+
+
+FString FConfigPropertyCellPresenter::GetValueAsString()
+{
+	return FString();
+}
+
+
+FText FConfigPropertyCellPresenter::GetValueAsText()
+{
+	return FText();
+}
+
+
+#undef LOCTEXT_NAMESPACE
\ No newline at end of file
diff --git a/Engine/Source/Editor/ConfigEditor/Private/PropertyVisualization/ConfigPropertyCellPresenter.h b/Engine/Source/Editor/ConfigEditor/Private/PropertyVisualization/ConfigPropertyCellPresenter.h
new file mode 100644
--- /dev/null
+++ b/Engine/Source/Editor/ConfigEditor/Private/PropertyVisualization/ConfigPropertyCellPresenter.h
@@ -0,0 +1,31 @@
+// Copyright 1998-2015 Epic Games, Inc. All Rights Reserved.
+
+#pragma once
+
+#include "Editor/PropertyEditor/Public/PropertyHandle.h"
+#include "Editor/PropertyEditor/Public/IPropertyTableCellPresenter.h"
+
+
+class FConfigPropertyCellPresenter : public TSharedFromThis< FConfigPropertyCellPresenter >, public IPropertyTableCellPresenter
+{
+public:
+
+	FConfigPropertyCellPresenter(const TSharedPtr< IPropertyHandle > PropertyHandle);
+	virtual ~FConfigPropertyCellPresenter() {}
+
+
+	/** Begin IPropertyTableCellPresenter interface */
+	virtual TSharedRef< class SWidget > ConstructDisplayWidget() override;
+	virtual bool RequiresDropDown() override;
+	virtual TSharedRef< class SWidget > ConstructEditModeCellWidget() override;
+	virtual TSharedRef< class SWidget > ConstructEditModeDropDownWidget() override;
+	virtual TSharedRef< class SWidget > WidgetToFocusOnEdit() override;
+	virtual bool HasReadOnlyEditMode() override;
+	virtual FString GetValueAsString() override;
+	virtual FText GetValueAsText() override;
+	/** End IPropertyTableCellPresenter interface */
+
+protected:
+	// The cell content
+	TSharedPtr<SWidget> DisplayWidget;
+};
\ No newline at end of file
diff --git a/Engine/Source/Editor/ConfigEditor/Private/PropertyVisualization/ConfigPropertyColumn.cpp b/Engine/Source/Editor/ConfigEditor/Private/PropertyVisualization/ConfigPropertyColumn.cpp
new file mode 100644
--- /dev/null
+++ b/Engine/Source/Editor/ConfigEditor/Private/PropertyVisualization/ConfigPropertyColumn.cpp
@@ -0,0 +1,60 @@
+// Copyright 1998-2015 Epic Games, Inc. All Rights Reserved.
+
+#include "ConfigEditorPCH.h"
+#include "ConfigPropertyColumn.h"
+#include "ConfigPropertyCellPresenter.h"
+
+#include "Editor/PropertyEditor/Public/PropertyHandle.h"
+#include "Editor/PropertyEditor/Public/PropertyPath.h"
+#include "Editor/PropertyEditor/Public/IPropertyTable.h"
+
+
+#define LOCTEXT_NAMESPACE "ConfigEditor"
+
+
+bool FConfigPropertyCustomColumn::Supports(const TSharedRef< IPropertyTableColumn >& Column, const TSharedRef< IPropertyTableUtilities >& Utilities) const
+{
+	bool IsSupported = false;
+
+	if (Column->GetDataSource()->IsValid())
+	{
+		TSharedPtr< FPropertyPath > PropertyPath = Column->GetDataSource()->AsPropertyPath();
+		if (PropertyPath.IsValid() && PropertyPath->GetNumProperties() > 0)
+		{
+			const FPropertyInfo& PropertyInfo = PropertyPath->GetRootProperty();
+			UProperty* Property = PropertyInfo.Property.Get();
+			IsSupported = Property->GetFName() == TEXT("ExternalProperty");
+		}
+	}
+
+	return IsSupported;
+}
+
+TSharedPtr< SWidget > FConfigPropertyCustomColumn::CreateColumnLabel(const TSharedRef< IPropertyTableColumn >& Column, const TSharedRef< IPropertyTableUtilities >& Utilities, const FName& Style) const
+{
+	if (Column->GetDataSource()->IsValid())
+	{
+		TSharedPtr< FPropertyPath > PropertyPath = Column->GetDataSource()->AsPropertyPath();
+		if (PropertyPath.IsValid() && PropertyPath->GetNumProperties() > 0)
+		{
+			return SNew(STextBlock)
+				.Text(EditProperty->GetDisplayNameText());
+		}
+	}
+	return SNullWidget::NullWidget;
+}
+
+
+TSharedPtr< IPropertyTableCellPresenter > FConfigPropertyCustomColumn::CreateCellPresenter(const TSharedRef< IPropertyTableCell >& Cell, const TSharedRef< IPropertyTableUtilities >& Utilities, const FName& Style) const
+{
+	TSharedPtr< IPropertyHandle > PropertyHandle = Cell->GetPropertyHandle();
+	if (PropertyHandle.IsValid())
+	{
+		return MakeShareable(new FConfigPropertyCellPresenter(PropertyHandle));
+	}
+
+	return nullptr;
+}
+
+
+#undef LOCTEXT_NAMESPACE
\ No newline at end of file
diff --git a/Engine/Source/Editor/ConfigEditor/Private/PropertyVisualization/ConfigPropertyConfigFileStateColumn.cpp b/Engine/Source/Editor/ConfigEditor/Private/PropertyVisualization/ConfigPropertyConfigFileStateColumn.cpp
new file mode 100644
--- /dev/null
+++ b/Engine/Source/Editor/ConfigEditor/Private/PropertyVisualization/ConfigPropertyConfigFileStateColumn.cpp
@@ -0,0 +1,120 @@
+// Copyright 1998-2015 Epic Games, Inc. All Rights Reserved.
+
+#include "ConfigEditorPCH.h"
+#include "ConfigPropertyConfigFileStateColumn.h"
+
+#include "Editor/PropertyEditor/Public/PropertyEditorModule.h"
+#include "Editor/PropertyEditor/Public/PropertyHandle.h"
+#include "Editor/PropertyEditor/Public/PropertyPath.h"
+#include "Editor/PropertyEditor/Public/IPropertyTable.h"
+
+
+#define LOCTEXT_NAMESPACE "ConfigEditor"
+
+
+FConfigPropertyConfigFileStateCellPresenter::FConfigPropertyConfigFileStateCellPresenter(const TSharedPtr< IPropertyHandle > PropertyHandle)
+{
+	//		Text = GetPropertyAsText(PropertyHandle);
+}
+
+
+TSharedRef< class SWidget > FConfigPropertyConfigFileStateCellPresenter::ConstructDisplayWidget()
+{
+	return
+		SNew(SHorizontalBox)
+		+ SHorizontalBox::Slot()
+		.AutoWidth()
+		[
+			SNew(STextBlock)
+			.Text(LOCTEXT("SCC unconnected","SCC"))
+		];
+}
+
+
+bool FConfigPropertyConfigFileStateCellPresenter::RequiresDropDown()
+{
+	return false;
+}
+
+
+TSharedRef< class SWidget > FConfigPropertyConfigFileStateCellPresenter::ConstructEditModeCellWidget()
+{
+	return ConstructDisplayWidget();
+}
+
+
+TSharedRef< class SWidget > FConfigPropertyConfigFileStateCellPresenter::ConstructEditModeDropDownWidget()
+{
+	return SNullWidget::NullWidget;
+}
+
+
+TSharedRef< class SWidget > FConfigPropertyConfigFileStateCellPresenter::WidgetToFocusOnEdit()
+{
+	return SNullWidget::NullWidget;
+}
+
+
+bool FConfigPropertyConfigFileStateCellPresenter::HasReadOnlyEditMode()
+{
+	return true;
+}
+
+
+FString FConfigPropertyConfigFileStateCellPresenter::GetValueAsString()
+{
+	return Text.ToString();
+}
+
+
+FText FConfigPropertyConfigFileStateCellPresenter::GetValueAsText()
+{
+	return Text;
+}
+
+
+bool FConfigPropertyConfigFileStateCustomColumn::Supports(const TSharedRef< IPropertyTableColumn >& Column, const TSharedRef< IPropertyTableUtilities >& Utilities) const
+{
+	bool IsSupported = false;
+
+	if (Column->GetDataSource()->IsValid())
+	{
+		TSharedPtr< FPropertyPath > PropertyPath = Column->GetDataSource()->AsPropertyPath();
+		if (PropertyPath.IsValid() && PropertyPath->GetNumProperties() > 0)
+		{
+			const FPropertyInfo& PropertyInfo = PropertyPath->GetRootProperty();
+			UProperty* Property = PropertyInfo.Property.Get();
+			if (SupportedProperty == Property)
+			{
+				IsSupported = true;
+			}
+		}
+	}
+
+	return IsSupported;
+}
+
+TSharedPtr< SWidget > FConfigPropertyConfigFileStateCustomColumn::CreateColumnLabel(const TSharedRef< IPropertyTableColumn >& Column, const TSharedRef< IPropertyTableUtilities >& Utilities, const FName& Style) const
+{
+	return NULL;
+}
+
+TSharedPtr< IPropertyTableCellPresenter > FConfigPropertyConfigFileStateCustomColumn::CreateCellPresenter(const TSharedRef< IPropertyTableCell >& Cell, const TSharedRef< IPropertyTableUtilities >& Utilities, const FName& Style) const
+{
+	TSharedPtr< IPropertyHandle > PropertyHandle = Cell->GetPropertyHandle();
+	if (PropertyHandle.IsValid())
+	{
+		TSharedRef<FPropertyPath> Path =  FPropertyPath::Create(PropertyHandle->GetProperty());
+		return MakeShareable(new FConfigPropertyConfigFileStateCellPresenter(PropertyHandle));
+	}
+
+	return NULL;
+	/*
+	SNew(SPropertyTableCell, Cell)
+		.Presenter(CellPresenter)
+		.Style(Style);
+	*/
+}
+
+
+#undef LOCTEXT_NAMESPACE
\ No newline at end of file
diff --git a/Engine/Source/Editor/ConfigEditor/Private/SConfigEditor.cpp b/Engine/Source/Editor/ConfigEditor/Private/SConfigEditor.cpp
new file mode 100644
--- /dev/null
+++ b/Engine/Source/Editor/ConfigEditor/Private/SConfigEditor.cpp
@@ -0,0 +1,99 @@
+// Copyright 1998-2015 Epic Games, Inc. All Rights Reserved.
+
+#include "ConfigEditorPCH.h"
+#include "ConfigPropertyHelper.h"
+#include "SConfigEditor.h"
+
+#include "PropertyEditorModule.h"
+
+#include "STargetPlatformSelector.h"
+
+
+#define LOCTEXT_NAMESPACE "ConfigEditor"
+
+
+BEGIN_SLATE_FUNCTION_BUILD_OPTIMIZATION
+void SConfigEditor::Construct(const FArguments& InArgs, TWeakObjectPtr<UProperty> InEditProperty)
+{
+	TargetPlatformSelection = SNew(STargetPlatformSelector)
+		.OnTargetPlatformChanged(this, &SConfigEditor::HandleTargetPlatformChanged);
+
+	EditProperty = InEditProperty;
+
+	LocalConfigCache = MakeShareable(new FConfigCacheIni(EConfigCacheType::Temporary));
+
+	// initialize details view
+	FDetailsViewArgs DetailsViewArgs;
+	{
+		DetailsViewArgs.bAllowSearch = false;
+		DetailsViewArgs.bHideSelectionTip = true;
+		DetailsViewArgs.bLockable = false;
+		DetailsViewArgs.bSearchInitialKeyFocus = true;
+		DetailsViewArgs.bUpdatesFromSelection = false;
+		DetailsViewArgs.bShowOptions = false;
+		DetailsViewArgs.bShowModifiedPropertiesOption = false;
+	}
+
+	DetailsView = FModuleManager::GetModuleChecked<FPropertyEditorModule>("PropertyEditor").CreateDetailView(DetailsViewArgs);
+	
+	CreateDisplayObjectForSelectedTargetPlatform();
+
+	////
+	// Our Widget setup is complete.
+	PropertyValueEditor = DetailsView.ToSharedRef();
+
+	ChildSlot
+	[
+		SNew(SVerticalBox)
+		+ SVerticalBox::Slot()
+		.AutoHeight()
+		[
+			TargetPlatformSelection.ToSharedRef()
+		]
+		+ SVerticalBox::Slot()
+		[
+			PropertyValueEditor.ToSharedRef()
+		]
+	];
+}
+END_SLATE_FUNCTION_BUILD_OPTIMIZATION
+
+void SConfigEditor::CreateDisplayObjectForSelectedTargetPlatform()
+{
+	// Get the ini file and hierarchy linked with this property
+	FString SelectedTargetPlatform = *TargetPlatformSelection->GetSelectedTargetPlatform().Get();
+
+	FString ConfigHelperName(TEXT("ConfigEditorPropertyHelper_"));
+	ConfigHelperName += SelectedTargetPlatform;
+
+	PropHelper = FindObject<UConfigHierarchyPropertyView>(GetTransientPackage(), *ConfigHelperName);
+	if (!PropHelper.IsValid())
+	{
+		PropHelper = NewObject<UConfigHierarchyPropertyView>(GetTransientPackage(), *ConfigHelperName);
+		PropHelper->AddToRoot();
+	}
+
+	PropHelper->EditProperty = EditProperty;
+	FString ClassConfigName = PropHelper->EditProperty->GetOwnerClass()->ClassConfigName.ToString();
+
+	FConfigFile PlatformIniFile;
+	LocalConfigCache->LoadLocalIniFile(PlatformIniFile, *ClassConfigName, true, *SelectedTargetPlatform);
+
+	for (auto& IniFile : PlatformIniFile.SourceIniHierarchy)
+	{
+		UPropertyConfigFileDisplayRow* ConfigFilePropertyObj = NewObject<UPropertyConfigFileDisplayRow>(GetTransientPackage(), *IniFile.Value.Filename);
+		ConfigFilePropertyObj->InitWithConfigAndProperty(IniFile.Value.Filename, PropHelper->EditProperty.Get());
+
+		PropHelper->ConfigFilePropertyObjects.Add(ConfigFilePropertyObj);
+	}
+
+	DetailsView->SetObject(PropHelper.Get());
+}
+
+void SConfigEditor::HandleTargetPlatformChanged()
+{
+	CreateDisplayObjectForSelectedTargetPlatform();
+}
+
+
+#undef LOCTEXT_NAMESPACE
diff --git a/Engine/Source/Editor/ConfigEditor/Private/SConfigEditor.h b/Engine/Source/Editor/ConfigEditor/Private/SConfigEditor.h
new file mode 100644
--- /dev/null
+++ b/Engine/Source/Editor/ConfigEditor/Private/SConfigEditor.h
@@ -0,0 +1,54 @@
+// Copyright 1998-2015 Epic Games, Inc. All Rights Reserved.
+
+#pragma once
+
+#include "ConfigPropertyHelper.h"
+#include "IDetailsView.h"
+
+#define LOCTEXT_NAMESPACE "ConfigEditor"
+
+/*-----------------------------------------------------------------------------
+   SConfigEditor
+-----------------------------------------------------------------------------*/
+class SConfigEditor : public SCompoundWidget
+{
+public:
+	SLATE_USER_ARGS(SConfigEditor) {}
+	SLATE_END_ARGS()
+
+	// Begin SCompoundWidget|SWidget interface
+	virtual void Construct(const FArguments& InArgs, TWeakObjectPtr<UProperty> InEditProperty);
+	// End SCompoundWidget|SWidget interface
+
+private:
+	/**
+	 * Handle a change of target platform in the Config Editor UI
+	 */
+	void HandleTargetPlatformChanged();
+	
+	// Panel used to select an available target platform.
+	TSharedPtr<class STargetPlatformSelector> TargetPlatformSelection;
+	
+	/**
+	 * Create the displayable area object for the selected platform
+	 */
+	void CreateDisplayObjectForSelectedTargetPlatform();
+
+private:
+	// This is an object we use to create a config hierarchy display for a property
+	TWeakObjectPtr<UConfigHierarchyPropertyView> PropHelper;
+
+	// The display area for the config hierarchy editor.
+	TSharedPtr<SWidget> PropertyValueEditor;
+
+	// We keep a cache of config files for the provided platforms.
+	TSharedPtr<FConfigCacheIni> LocalConfigCache;
+
+	// The bulk of the display of this hierarchy.
+	TSharedPtr<IDetailsView> DetailsView;
+
+	// Keep track of the property we are viewing.
+	TWeakObjectPtr<UProperty> EditProperty;
+};
+
+#undef LOCTEXT_NAMESPACE // "ConfigEditor"
\ No newline at end of file
diff --git a/Engine/Source/Editor/ConfigEditor/Private/STargetPlatformSelector.cpp b/Engine/Source/Editor/ConfigEditor/Private/STargetPlatformSelector.cpp
new file mode 100644
--- /dev/null
+++ b/Engine/Source/Editor/ConfigEditor/Private/STargetPlatformSelector.cpp
@@ -0,0 +1,96 @@
+// Copyright 1998-2015 Epic Games, Inc. All Rights Reserved.
+
+#include "ConfigEditorPCH.h"
+#include "STargetPlatformSelector.h"
+#include "ModuleManager.h"
+
+#include "ITargetPlatformManagerModule.h"
+#include "ITargetPlatform.h"
+
+#include "STextComboBox.h"
+
+#define LOCTEXT_NAMESPACE "ConfigEditor"
+
+
+BEGIN_SLATE_FUNCTION_BUILD_OPTIMIZATION
+void STargetPlatformSelector::Construct(const FArguments& InArgs)
+{
+	OnTargetPlatformChanged = InArgs._OnTargetPlatformChanged;
+
+	CollateAvailableTargetPlatformEntries();
+
+	this->ChildSlot
+	[
+		SNew(SBorder)
+		.BorderImage(FCoreStyle::Get().GetBrush("ToolPanel.GroupBorder"))
+		.Padding(FMargin(8.0f, 8.0f, 16.0f, 8.0f))
+		[
+			SNew(SHorizontalBox)
+			+ SHorizontalBox::Slot()
+			.HAlign(EHorizontalAlignment::HAlign_Right)
+			.VAlign(EVerticalAlignment::VAlign_Center)
+			.Padding(FMargin(8.0f, 8.0f, 8.0f, 8.0f))
+			[
+				SNew(STextBlock)
+				.Text(LOCTEXT("TargetPlatformSelectionLabel", "Target platform:"))
+			]
+			+ SHorizontalBox::Slot()
+			.HAlign(EHorizontalAlignment::HAlign_Right)
+			.Padding(FMargin(0.0f, 4.0f))
+			.AutoWidth()
+			[
+				SNew(SBox)
+				.WidthOverride(150.0f)
+				[
+					SAssignNew(AvailableTargetPlatformComboBox, STextComboBox)
+					.OptionsSource(&TargetPlatformOptionsSource)
+					.InitiallySelectedItem(SelectedTargetPlatform)
+					.OnSelectionChanged(this, &STargetPlatformSelector::HandleTargetPlatformChanged)
+				]
+			]
+		]
+	];
+
+}
+END_SLATE_FUNCTION_BUILD_OPTIMIZATION
+
+
+TSharedPtr<FString> STargetPlatformSelector::GetSelectedTargetPlatform()
+{
+	return SelectedTargetPlatform;
+}
+
+
+void STargetPlatformSelector::HandleTargetPlatformChanged(TSharedPtr<FString> NewValue, ESelectInfo::Type SelectInfo)
+{
+	SelectedTargetPlatform = NewValue;
+	OnTargetPlatformChanged.ExecuteIfBound();
+}
+
+
+void STargetPlatformSelector::CollateAvailableTargetPlatformEntries()
+{
+	ITargetPlatformManagerModule& TPM = GetTargetPlatformManagerRef();
+	TArray<ITargetPlatform*> ActiveTargetPlatforms = TPM.GetTargetPlatforms();
+	for (ITargetPlatform* TargetPlatform : ActiveTargetPlatforms)
+	{
+		const FString CurrentIniName = TargetPlatform->IniPlatformName();
+		if (!CurrentIniName.IsEmpty())
+		{
+			const TSharedPtr<FString>* ExistingEntry = TargetPlatformOptionsSource.FindByPredicate([CurrentIniName](const TSharedPtr<FString>& CurElement)
+			{
+				return *CurElement == CurrentIniName;
+			});
+
+			if (ExistingEntry == nullptr || !(*ExistingEntry).IsValid())
+			{
+				TargetPlatformOptionsSource.Add(MakeShareable(new FString(TargetPlatform->IniPlatformName())));
+			}
+		}
+	}
+
+	SelectedTargetPlatform = TargetPlatformOptionsSource[0];
+}
+
+
+#undef LOCTEXT_NAMESPACE
diff --git a/Engine/Source/Editor/ConfigEditor/Private/STargetPlatformSelector.h b/Engine/Source/Editor/ConfigEditor/Private/STargetPlatformSelector.h
new file mode 100644
--- /dev/null
+++ b/Engine/Source/Editor/ConfigEditor/Private/STargetPlatformSelector.h
@@ -0,0 +1,55 @@
+// Copyright 1998-2015 Epic Games, Inc. All Rights Reserved.
+
+#pragma once
+
+#define LOCTEXT_NAMESPACE "ConfigEditor"
+
+/*-----------------------------------------------------------------------------
+   STargetPlatformSelector
+-----------------------------------------------------------------------------*/
+class STargetPlatformSelector : public SCompoundWidget
+{
+public:
+	SLATE_BEGIN_ARGS(STargetPlatformSelector){}
+		SLATE_EVENT(FSimpleDelegate, OnTargetPlatformChanged)
+	SLATE_END_ARGS()
+
+	// Begin SCompoundWidget|SWidget interface
+	void Construct(const FArguments& InArgs);
+	// End SCompoundWidget|SWidget interface
+
+	/**
+	 * Access to the selected target platform in this widget.
+	 *
+	 * @return - The selected platform from the target platform ComboBox.
+	 */
+	TSharedPtr<FString> GetSelectedTargetPlatform();
+
+private:
+	 
+	/**
+	 * Callback when the user has made a change to the target platform.
+	 */
+	void HandleTargetPlatformChanged(TSharedPtr<FString> NewValue, ESelectInfo::Type SelectInfo);
+
+	/**
+	 * Build a list of available target platforms.
+	 */
+	void CollateAvailableTargetPlatformEntries();
+
+private:
+
+	// The currently selected target platform in this widget.
+	TSharedPtr<FString> SelectedTargetPlatform;
+
+	// The combo widget used to make a change to selected target platform.
+	TSharedPtr<STextComboBox> AvailableTargetPlatformComboBox;
+
+	// The list of available target platforms for the combo widget.
+	TArray< TSharedPtr<FString> > TargetPlatformOptionsSource;
+
+	// Delegate called to let the listener know there has been a change to the target platform.
+	FSimpleDelegate OnTargetPlatformChanged;
+};
+
+#undef LOCTEXT_NAMESPACE // "ConfigEditor"
\ No newline at end of file
diff --git a/Engine/Source/Editor/ConfigEditor/Public/ConfigPropertyHelper.h b/Engine/Source/Editor/ConfigEditor/Public/ConfigPropertyHelper.h
new file mode 100644
--- /dev/null
+++ b/Engine/Source/Editor/ConfigEditor/Public/ConfigPropertyHelper.h
@@ -0,0 +1,46 @@
+// Copyright 1998-2015 Epic Games, Inc. All Rights Reserved.
+
+#pragma once
+ 
+#include "ConfigPropertyHelper.generated.h"
+
+UENUM()
+enum EConfigFileSourceControlStatus
+{
+	CFSCS_Unknown UMETA(DisplayName=Unknown),
+	CFSCS_Writable UMETA(DisplayName = "Available to edit"),
+	CFSCS_Locked UMETA(DisplayName = "File is locked"),
+};
+
+UCLASS(MinimalAPI)
+class UPropertyConfigFileDisplayRow : public UObject
+{
+	GENERATED_BODY()
+public:
+
+	void InitWithConfigAndProperty(const FString& InConfigFileName, UProperty* InEditProperty);
+
+public:
+	UPROPERTY(Transient, Category = Helper, VisibleAnywhere)
+	FString ConfigFileName;
+
+	UPROPERTY(Transient, Category = Helper, EditAnywhere, meta=(EditCondition="bIsFileWritable"))
+	UProperty* ExternalProperty;
+
+	UPROPERTY(Transient, Category = Helper, VisibleAnywhere)
+	bool bIsFileWritable;
+};
+
+
+UCLASS(MinimalAPI)
+class UConfigHierarchyPropertyView : public UObject
+{
+	GENERATED_BODY()
+
+public:
+	UPROPERTY(Transient, Category=Helper, EditAnywhere)
+	TWeakObjectPtr<UProperty> EditProperty;
+
+	UPROPERTY(Transient, Category = Helper, EditAnywhere)
+	TArray<UPropertyConfigFileDisplayRow*> ConfigFilePropertyObjects;
+};
diff --git a/Engine/Source/Editor/ConfigEditor/Public/IConfigEditorModule.h b/Engine/Source/Editor/ConfigEditor/Public/IConfigEditorModule.h
new file mode 100644
--- /dev/null
+++ b/Engine/Source/Editor/ConfigEditor/Public/IConfigEditorModule.h
@@ -0,0 +1,29 @@
+// Copyright 1998-2015 Epic Games, Inc. All Rights Reserved.
+
+#pragma once
+
+
+/*-----------------------------------------------------------------------------
+   IConfigEditorModule
+-----------------------------------------------------------------------------*/
+class IConfigEditorModule : public IModuleInterface
+{
+public:
+	/**
+	 * Creates a config hierarchy editor for the given property 
+	 */
+	virtual void CreateHierarchyEditor(UProperty* EditProperty) = 0;
+
+	/**
+	 * Maintain a reference to a value widget used in the config editor for the current property.
+	 */
+	virtual void AddExternalPropertyValueWidgetAndConfigPairing(const FString& ConfigFile, const TSharedPtr<SWidget> ValueWidget) = 0;
+
+	/**
+	 * Access a reference to a value widget used in the config editor for the current property.
+	 *
+	 * @return -	A widget which represents the config value for a specific property. This is a different widget based on the property type
+	 *				ie. a combobox for enum, a check box for bool...
+	 */
+	virtual TSharedRef<SWidget> GetValueWidgetForConfigProperty(const FString& ConfigFile) = 0;
+};
diff --git a/Engine/Source/Editor/ConfigEditor/Public/PropertyVisualization/ConfigPropertyColumn.h b/Engine/Source/Editor/ConfigEditor/Public/PropertyVisualization/ConfigPropertyColumn.h
new file mode 100644
--- /dev/null
+++ b/Engine/Source/Editor/ConfigEditor/Public/PropertyVisualization/ConfigPropertyColumn.h
@@ -0,0 +1,25 @@
+// Copyright 1998-2015 Epic Games, Inc. All Rights Reserved.
+
+#pragma once
+
+#include "Editor/PropertyEditor/Public/IPropertyTableColumn.h"
+#include "Editor/PropertyEditor/Public/IPropertyTableCustomColumn.h"
+
+
+/**
+* A property table custom column used to display source control condition of files.
+*/
+class CONFIGEDITOR_API FConfigPropertyCustomColumn : public IPropertyTableCustomColumn
+{
+public:
+	FConfigPropertyCustomColumn(){}
+
+	/** Begin IPropertyTableCustomColumn interface */
+	virtual bool Supports(const TSharedRef< IPropertyTableColumn >& Column, const TSharedRef< IPropertyTableUtilities >& Utilities) const override;
+	virtual TSharedPtr< SWidget > CreateColumnLabel(const TSharedRef< IPropertyTableColumn >& Column, const TSharedRef< IPropertyTableUtilities >& Utilities, const FName& Style) const override;
+	virtual TSharedPtr< IPropertyTableCellPresenter > CreateCellPresenter(const TSharedRef< IPropertyTableCell >& Cell, const TSharedRef< IPropertyTableUtilities >& Utilities, const FName& Style) const override;
+	/** End IPropertyTableCustomColumn interface */
+
+	/* The property type which can be displayed in this column */
+	UProperty* EditProperty;
+};
\ No newline at end of file
diff --git a/Engine/Source/Editor/ConfigEditor/Public/PropertyVisualization/ConfigPropertyConfigFileStateColumn.h b/Engine/Source/Editor/ConfigEditor/Public/PropertyVisualization/ConfigPropertyConfigFileStateColumn.h
new file mode 100644
--- /dev/null
+++ b/Engine/Source/Editor/ConfigEditor/Public/PropertyVisualization/ConfigPropertyConfigFileStateColumn.h
@@ -0,0 +1,57 @@
+// Copyright 1998-2015 Epic Games, Inc. All Rights Reserved.
+
+#pragma once
+
+#include "Editor/PropertyEditor/Public/IPropertyTableColumn.h"
+#include "Editor/PropertyEditor/Public/IPropertyTableCustomColumn.h"
+#include "Editor/PropertyEditor/Public/IPropertyTableCellPresenter.h"
+
+
+class FConfigPropertyConfigFileStateCellPresenter : public TSharedFromThis< FConfigPropertyConfigFileStateCellPresenter >, public IPropertyTableCellPresenter
+{
+public:
+
+	FConfigPropertyConfigFileStateCellPresenter(const TSharedPtr< IPropertyHandle > PropertyHandle);
+	virtual ~FConfigPropertyConfigFileStateCellPresenter() {}
+
+
+	/** Begin IPropertyTableCellPresenter interface */
+	virtual TSharedRef< class SWidget > ConstructDisplayWidget() override;
+	virtual bool RequiresDropDown() override;
+	virtual TSharedRef< class SWidget > ConstructEditModeCellWidget() override;
+	virtual TSharedRef< class SWidget > ConstructEditModeDropDownWidget() override;
+	virtual TSharedRef< class SWidget > WidgetToFocusOnEdit() override;
+	virtual bool HasReadOnlyEditMode() override;
+	virtual FString GetValueAsString() override;
+
+	virtual FText GetValueAsText() override;
+	/** End IPropertyTableCellPresenter interface */
+
+protected:
+
+	/**
+	* The text we display - we just store this as we never need to *edit* properties in the stats viewer,
+	* only view them. Therefore we dont need to dynamically update this string.
+	*/
+	FText Text;
+};
+
+
+
+/**
+* A property table custom column used to display source control condition of files.
+*/
+class CONFIGEDITOR_API FConfigPropertyConfigFileStateCustomColumn : public IPropertyTableCustomColumn
+{
+public:
+	FConfigPropertyConfigFileStateCustomColumn(){}
+
+	/** Begin IPropertyTableCustomColumn interface */
+	virtual bool Supports(const TSharedRef< IPropertyTableColumn >& Column, const TSharedRef< IPropertyTableUtilities >& Utilities) const override;
+	virtual TSharedPtr< SWidget > CreateColumnLabel(const TSharedRef< IPropertyTableColumn >& Column, const TSharedRef< IPropertyTableUtilities >& Utilities, const FName& Style) const override;
+	virtual TSharedPtr< IPropertyTableCellPresenter > CreateCellPresenter(const TSharedRef< IPropertyTableCell >& Cell, const TSharedRef< IPropertyTableUtilities >& Utilities, const FName& Style) const override;
+	/** End IPropertyTableCustomColumn interface */
+
+	/* The property type which can be displayed in this column */
+	UProperty* SupportedProperty;
+};
\ No newline at end of file
diff --git a/Engine/Source/Editor/DetailCustomizations/DetailCustomizations.Build.cs b/Engine/Source/Editor/DetailCustomizations/DetailCustomizations.Build.cs
--- a/Engine/Source/Editor/DetailCustomizations/DetailCustomizations.Build.cs
+++ b/Engine/Source/Editor/DetailCustomizations/DetailCustomizations.Build.cs
@@ -42,7 +42,8 @@ public class DetailCustomizations : ModuleRules
                 "HardwareTargeting",
 				"SharedSettingsWidgets",
                 "AIModule", 
-                "MeshUtilities"
+                "MeshUtilities",
+				"ConfigEditor",
 			}
 		);
 
diff --git a/Engine/Source/Editor/DetailCustomizations/Private/ConfigEditorPropertyDetails.cpp b/Engine/Source/Editor/DetailCustomizations/Private/ConfigEditorPropertyDetails.cpp
new file mode 100644
--- /dev/null
+++ b/Engine/Source/Editor/DetailCustomizations/Private/ConfigEditorPropertyDetails.cpp
@@ -0,0 +1,181 @@
+// Copyright 1998-2015 Epic Games, Inc. All Rights Reserved.
+
+#include "DetailCustomizationsPrivatePCH.h"
+#include "ConfigEditorPropertyDetails.h"
+#include "Editor/PropertyEditor/Public/PropertyEditorModule.h"
+#include "Editor/PropertyEditor/Public/PropertyPath.h"
+#include "Editor/PropertyEditor/Public/IPropertyTable.h"
+#include "Editor/PropertyEditor/Public/IPropertyTableColumn.h"
+#include "Editor/PropertyEditor/Public/IPropertyTableCustomColumn.h"
+
+#include "IConfigEditorModule.h"
+#include "ConfigPropertyColumn.h"
+#include "ConfigPropertyConfigFileStateColumn.h"
+#include "ConfigPropertyHelper.h"
+
+
+#define LOCTEXT_NAMESPACE "ConfigPropertyHelperDetails"
+
+////////////////////////////////////////////////
+// FConfigPropertyHelperDetails
+
+void FConfigPropertyHelperDetails::CustomizeDetails(IDetailLayoutBuilder& DetailBuilder)
+{
+	TSharedPtr<IPropertyHandle> PropertyHandle = DetailBuilder.GetProperty("EditProperty");
+	DetailBuilder.HideProperty(PropertyHandle);
+
+	UObject* PropValue;
+	PropertyHandle->GetValue(PropValue);
+
+	// Keep a record of the UProperty we are looking to update
+	Property = CastChecked<UProperty>(PropValue);
+
+	// Get access to all of the config files where this property is configurable.
+	ConfigFilesHandle = DetailBuilder.GetProperty("ConfigFilePropertyObjects");
+	FSimpleDelegate OnConfigFileListChangedDelegate = FSimpleDelegate::CreateSP(this, &FConfigPropertyHelperDetails::OnConfigFileListChanged);
+	ConfigFilesHandle->SetOnPropertyValueChanged(OnConfigFileListChangedDelegate);
+	DetailBuilder.HideProperty(ConfigFilesHandle);
+
+	// Add the properties to a property table so we can edit these.
+	IDetailCategoryBuilder& A = DetailBuilder.EditCategory("ConfigHierarchy");
+	A.AddCustomRow(LOCTEXT("ConfigHierarchy", "ConfigHierarchy"))
+	[
+		// Create a property table with the values.
+		ConstructPropertyTable(DetailBuilder)
+	];
+
+	// Listen for changes to the properties, we handle these by updating the ini file associated.
+	FCoreUObjectDelegates::OnObjectPropertyChanged.AddSP(this, &FConfigPropertyHelperDetails::OnPropertyValueChanged);
+}
+
+
+void FConfigPropertyHelperDetails::OnPropertyValueChanged(UObject* Object, FPropertyChangedEvent& PropertyChangedEvent)
+{
+	UClass* OwnerClass = Property->GetOwnerClass();
+	if (Object->IsA(OwnerClass))
+	{
+		const FString* FileName = ConfigFileAndPropertySourcePairings.FindKey(Object);
+		if (FileName != nullptr)
+		{
+			Object->UpdateDefaultConfigFile(*FileName);
+		}
+	}
+}
+
+
+// TPMB - This didnt work... why?
+void FConfigPropertyHelperDetails::OnConfigFileListChanged()
+{
+	PropertyTable->RequestRefresh();
+}
+
+
+void FConfigPropertyHelperDetails::AddEditablePropertyForConfig(IDetailLayoutBuilder& DetailBuilder, const UPropertyConfigFileDisplayRow* ConfigFileProperty)
+{
+	AssociatedConfigFileAndObjectPairings.Add(ConfigFileProperty->ConfigFileName, (UObject*)ConfigFileProperty);
+
+	// Add the properties to a property table so we can edit these.
+	IDetailCategoryBuilder& A = DetailBuilder.EditCategory("TestCat");
+
+	// Create an instance of the owner object so we can dictate values for our property on a per object basis.
+	UObject* CDO = Property->GetOwnerClass()->GetDefaultObject<UObject>();
+	UObject* ConfigEntryObject = StaticDuplicateObject(CDO, GetTransientPackage(), *(ConfigFileProperty->ConfigFileName + TEXT("_cdoDupe")));
+	ConfigEntryObject->AddToRoot();
+	ConfigEntryObject->LoadConfig(Property->GetOwnerClass(), *ConfigFileProperty->ConfigFileName, UE4::LCPF_None, Property);
+
+	// Add a reference for future saving.
+	ConfigFileAndPropertySourcePairings.Add(ConfigFileProperty->ConfigFileName, (UObject*)ConfigEntryObject);
+
+	// We need to add a property row for each config file entry. 
+	// This allows us to have an editable widget for each config file.
+	TArray<UObject*> ConfigPropertyDisplayObjects;
+	ConfigPropertyDisplayObjects.Add(ConfigEntryObject);
+	if (IDetailPropertyRow* ExternalRow = A.AddExternalProperty(ConfigPropertyDisplayObjects, Property->GetFName()))
+	{
+		TSharedPtr<SWidget> NameWidget;
+		TSharedPtr<SWidget> ValueWidget;
+		ExternalRow->GetDefaultWidgets(NameWidget, ValueWidget);
+
+		// Register the Value widget and config file pairing with the config editor.
+		// The config editor needs this to determine what a cell presenter shows.
+		IConfigEditorModule& ConfigEditor = FModuleManager::Get().LoadModuleChecked<IConfigEditorModule>("ConfigEditor");
+		ConfigEditor.AddExternalPropertyValueWidgetAndConfigPairing(ConfigFileProperty->ConfigFileName, ValueWidget);
+		
+		// now hide the property so it is not added to the property display view
+		ExternalRow->Visibility(EVisibility::Hidden);
+	}
+}
+
+
+TSharedRef<SWidget> FConfigPropertyHelperDetails::ConstructPropertyTable(IDetailLayoutBuilder& DetailBuilder)
+{
+	FPropertyEditorModule& PropertyEditorModule = FModuleManager::LoadModuleChecked<FPropertyEditorModule>("PropertyEditor");
+
+	PropertyTable = PropertyEditorModule.CreatePropertyTable();
+	PropertyTable->SetSelectionMode(ESelectionMode::None);
+	PropertyTable->SetIsUserAllowedToChangeRoot(false);
+	PropertyTable->SetShowObjectName(false);
+
+	RepopulatePropertyTable(DetailBuilder);
+
+	TArray< TSharedRef<class IPropertyTableCustomColumn>> CustomColumns;
+	TSharedRef<FConfigPropertyCustomColumn> EditPropertyColumn = MakeShareable(new FConfigPropertyCustomColumn());
+	EditPropertyColumn->EditProperty = Property; // FindFieldChecked<UProperty>(UPropertyConfigFileDisplay::StaticClass(), TEXT("EditProperty"));
+	CustomColumns.Add(EditPropertyColumn);
+
+	//TSharedRef<FConfigPropertyConfigFileStateCustomColumn> ConfigFileStateColumn = MakeShareable(new FConfigPropertyConfigFileStateCustomColumn());
+	//ConfigFileStateColumn->SupportedProperty = FindFieldChecked<UProperty>(UPropertyConfigFileDisplay::StaticClass(), TEXT("FileState"));
+	//CustomColumns.Add(ConfigFileStateColumn);
+
+	return PropertyEditorModule.CreatePropertyTableWidget(PropertyTable.ToSharedRef(), CustomColumns);
+}
+
+
+void FConfigPropertyHelperDetails::RepopulatePropertyTable(IDetailLayoutBuilder& DetailBuilder)
+{
+	// Clear out any previous entries from the table.
+	AssociatedConfigFileAndObjectPairings.Empty();
+
+	// Add an entry for each config so the value can be set in each of the config files independently.
+	uint32 ConfigCount = 0;
+	TSharedPtr<IPropertyHandleArray> ConfigFilesArrayHandle = ConfigFilesHandle->AsArray();
+	ConfigFilesArrayHandle->GetNumElements(ConfigCount);
+
+	// For each config file, add the capacity to edit this property.
+	for (uint32 Index = 0; Index < ConfigCount; Index++)
+	{
+		FString ConfigFile;
+		TSharedRef<IPropertyHandle> ConfigFileElementHandle = ConfigFilesArrayHandle->GetElement(Index);
+
+		UObject* ConfigFileSinglePropertyHelperObj = nullptr;
+		ensure(ConfigFileElementHandle->GetValue(ConfigFileSinglePropertyHelperObj) == FPropertyAccess::Success);
+
+		UPropertyConfigFileDisplayRow* ConfigFileSinglePropertyHelper = CastChecked<UPropertyConfigFileDisplayRow>(ConfigFileSinglePropertyHelperObj);
+
+		AddEditablePropertyForConfig(DetailBuilder, ConfigFileSinglePropertyHelper);
+	}
+
+	// We need a row for each config file
+	TArray<UObject*> ConfigPropertyDisplayObjects;
+	AssociatedConfigFileAndObjectPairings.GenerateValueArray(ConfigPropertyDisplayObjects);
+	PropertyTable->SetObjects(ConfigPropertyDisplayObjects);
+
+	// We need a column for each property in our Helper class.
+	for (UProperty* NextProperty = UPropertyConfigFileDisplayRow::StaticClass()->PropertyLink; NextProperty; NextProperty = NextProperty->PropertyLinkNext)
+	{
+		PropertyTable->AddColumn((TWeakObjectPtr<UProperty>)NextProperty);
+	}
+
+	// Ensure the columns cannot be removed.
+	TArray<TSharedRef<IPropertyTableColumn>> Columns = PropertyTable->GetColumns();
+	for (TSharedRef<IPropertyTableColumn> Column : Columns)
+	{
+		Column->SetFrozen(true);
+	}
+
+	// Create the 'Config File' vs 'Property' table
+	PropertyTable->RequestRefresh();
+}
+
+
+#undef LOCTEXT_NAMESPACE
\ No newline at end of file
diff --git a/Engine/Source/Editor/DetailCustomizations/Private/ConfigEditorPropertyDetails.h b/Engine/Source/Editor/DetailCustomizations/Private/ConfigEditorPropertyDetails.h
new file mode 100644
--- /dev/null
+++ b/Engine/Source/Editor/DetailCustomizations/Private/ConfigEditorPropertyDetails.h
@@ -0,0 +1,83 @@
+// Copyright 1998-2015 Epic Games, Inc. All Rights Reserved.
+
+#pragma once
+
+
+/**
+ * Implements details panel customizations for UConfigPropertyHelper fields.
+ */
+class FConfigPropertyHelperDetails
+	: public IDetailCustomization
+{
+	//////////////////////////////////////////////////////////////////////////
+	/// Detail Customization
+public:
+
+	/**
+	 * Makes a new instance of this config editor detail layout class.
+	 *
+	 * @return The created instance.
+	 */
+	static TSharedRef<IDetailCustomization> MakeInstance()
+	{
+		return MakeShareable(new FConfigPropertyHelperDetails);
+	}
+
+	// IDetailCustomization interface
+	virtual void CustomizeDetails( IDetailLayoutBuilder& InDetailBuilder ) override;
+
+
+	//////////////////////////////////////////////////////////////////////////
+	/// Property Table functionality
+private:
+
+	/**
+	 * Create a property table of the Config Files vs Property we are editing
+	 *
+	 * @return The created instance of the property table.
+	 */
+	TSharedRef<SWidget> ConstructPropertyTable(IDetailLayoutBuilder& DetailBuilder);
+
+	/**
+	 * Populate the Property table with entries for the provided config files.
+	 */
+	void RepopulatePropertyTable(IDetailLayoutBuilder& DetailBuilder);
+
+	/**
+	 * Event which is triggered through changes in our editor.
+	 *
+	 * @param Object				- The object whose property has changed.
+	 * @param PropertyChangedEvent	- The event information of the change.
+	 */
+	void OnPropertyValueChanged(UObject* Object, FPropertyChangedEvent& PropertyChangedEvent);
+
+private:
+	
+	// The table which holds our editable properties.
+	TSharedPtr<IPropertyTable> PropertyTable;
+
+
+	//////////////////////////////////////////////////////////////////////////
+	/// Misc operations...
+private:
+	void AddEditablePropertyForConfig(IDetailLayoutBuilder& DetailBuilder, const class UPropertyConfigFileDisplayRow* ConfigFileProperty);
+
+	void OnConfigFileListChanged();
+
+	//////////////////////////////////////////////////////////////////////////
+	/// Misc properties...
+private:
+
+	// Keep track of the property handle for the config files.
+	TSharedPtr<IPropertyHandle> ConfigFilesHandle;
+
+	// The property we wish to edit on a per config basis.
+	UProperty* Property;
+
+	// Coupling of Config files and their editable objects.
+	TMap<FString, UObject*> AssociatedConfigFileAndObjectPairings;
+
+	TMap<FString, UObject*> ConfigFileAndPropertySourcePairings;
+
+	TMap<FString, TSharedPtr<SWidget>> ExternalPropertyValueWidgetAndConfigPairings;
+};
diff --git a/Engine/Source/Editor/DetailCustomizations/Private/DetailCustomizations.cpp b/Engine/Source/Editor/DetailCustomizations/Private/DetailCustomizations.cpp
--- a/Engine/Source/Editor/DetailCustomizations/Private/DetailCustomizations.cpp
+++ b/Engine/Source/Editor/DetailCustomizations/Private/DetailCustomizations.cpp
@@ -101,6 +101,7 @@
 #include "DistanceDatumStructCustomization.h"
 #include "HierarchicalSimplificationCustomizations.h"
 #include "PostProcessSettingsCustomization.h"
+#include "ConfigEditorPropertyDetails.h"
 
 IMPLEMENT_MODULE( FDetailCustomizationsModule, DetailCustomizations );
 
@@ -292,6 +293,8 @@ void FDetailCustomizationsModule::RegisterObjectCustomizations()
 	RegisterCustomClassLayout("MediaTexture", FOnGetDetailCustomizationInstance::CreateStatic(&FMediaTextureCustomization::MakeInstance));
 
 	RegisterCustomClassLayout("FbxImportUI", FOnGetDetailCustomizationInstance::CreateStatic(&FFbxImportUIDetails::MakeInstance));
+
+	RegisterCustomClassLayout("ConfigHierarchyPropertyView", FOnGetDetailCustomizationInstance::CreateStatic(&FConfigPropertyHelperDetails::MakeInstance));
 }
 
 
diff --git a/Engine/Source/Editor/PropertyEditor/Private/Presentation/PropertyEditor/PropertyEditor.cpp b/Engine/Source/Editor/PropertyEditor/Private/Presentation/PropertyEditor/PropertyEditor.cpp
--- a/Engine/Source/Editor/PropertyEditor/Private/Presentation/PropertyEditor/PropertyEditor.cpp
+++ b/Engine/Source/Editor/PropertyEditor/Private/Presentation/PropertyEditor/PropertyEditor.cpp
@@ -13,6 +13,7 @@
 
 #include "EditorClassUtils.h"
 #include "KismetEditorUtilities.h"
+#include "IConfigEditorModule.h"
 
 #define LOCTEXT_NAMESPACE "PropertyEditor"
 
@@ -231,6 +232,15 @@ void FPropertyEditor::MakeNewBlueprint()
 	}
 }
 
+void FPropertyEditor::EditConfigHierarchy()
+{
+	UProperty* NodeProperty = PropertyNode->GetProperty();
+
+	IConfigEditorModule& ConfigEditorModule = FModuleManager::LoadModuleChecked<IConfigEditorModule>("ConfigEditor");
+	ConfigEditorModule.CreateHierarchyEditor(NodeProperty);
+	FGlobalTabmanager::Get()->InvokeTab(FName(TEXT("ConfigEditor")));
+}
+
 void FPropertyEditor::InsertItem()
 {
 	// This action must be deferred until next tick so that we avoid accessing invalid data before we have a chance to tick
diff --git a/Engine/Source/Editor/PropertyEditor/Private/Presentation/PropertyEditor/PropertyEditor.h b/Engine/Source/Editor/PropertyEditor/Private/Presentation/PropertyEditor/PropertyEditor.h
--- a/Engine/Source/Editor/PropertyEditor/Private/Presentation/PropertyEditor/PropertyEditor.h
+++ b/Engine/Source/Editor/PropertyEditor/Private/Presentation/PropertyEditor/PropertyEditor.h
@@ -92,6 +92,7 @@ public:
 	void OnAssetSelected( const FAssetData& AssetData );
 	void OnActorSelected( AActor* InActor );
 	void OnGetActorFiltersForSceneOutliner( TSharedPtr<SceneOutliner::FOutlinerFilters>& OutFilters );
+	void EditConfigHierarchy();
 
 	/**	In an ideal world we wouldn't expose these */
 	TSharedRef< FPropertyNode > GetPropertyNode() const;
diff --git a/Engine/Source/Editor/PropertyEditor/Private/PropertyCustomizationHelpers.cpp b/Engine/Source/Editor/PropertyEditor/Private/PropertyCustomizationHelpers.cpp
--- a/Engine/Source/Editor/PropertyEditor/Private/PropertyCustomizationHelpers.cpp
+++ b/Engine/Source/Editor/PropertyEditor/Private/PropertyCustomizationHelpers.cpp
@@ -248,6 +248,18 @@ namespace PropertyCustomizationHelpers
 			.OnActorSelected( OnActorSelectedFromPicker );
 	}
 
+	TSharedRef<SWidget> MakeEditConfigHierarchyButton(FSimpleDelegate OnEditConfigClicked, TAttribute<FText> OptionalToolTipText, TAttribute<bool> IsEnabled)
+	{
+		return
+			SNew(SPropertyEditorButton)
+			.Text(LOCTEXT("EditConfigHierarchyButtonLabel", "Edit Config Hierarchy"))
+			.ToolTipText(OptionalToolTipText.Get().IsEmpty() ? LOCTEXT("EditConfigHierarchyButtonToolTipText", "Edit the config values of this property") : OptionalToolTipText)
+			.Image(FEditorStyle::GetBrush("DetailsView.EditRawProperties"))
+			.OnClickAction(OnEditConfigClicked)
+			.IsEnabled(IsEnabled)
+			.IsFocusable(false);
+	}
+
 	UBoolProperty* GetEditConditionProperty(const UProperty* InProperty, bool& bNegate)
 	{
 		UBoolProperty* EditConditionProperty = NULL;
diff --git a/Engine/Source/Editor/PropertyEditor/Private/PropertyEditorHelpers.cpp b/Engine/Source/Editor/PropertyEditor/Private/PropertyEditorHelpers.cpp
--- a/Engine/Source/Editor/PropertyEditor/Private/PropertyEditorHelpers.cpp
+++ b/Engine/Source/Editor/PropertyEditor/Private/PropertyEditorHelpers.cpp
@@ -496,6 +496,14 @@ namespace PropertyEditorHelpers
 			return;
 		}
 
+		if (NodeProperty->HasAnyPropertyFlags(CPF_GlobalConfig|CPF_Config))
+		{
+			if (NodeProperty->HasMetaData(TEXT("ConfigHierarchyEditable")))
+			{
+				OutRequiredButtons.Add( EPropertyButton::EditConfigHierarchy );
+			}
+		}
+
 		//////////////////////////////
 		// Handle an array property.
 		if( NodeProperty->IsA(UArrayProperty::StaticClass() ) )
@@ -804,6 +812,10 @@ namespace PropertyEditorHelpers
 			NewButton = PropertyCustomizationHelpers::MakeNewBlueprintButton( FSimpleDelegate::CreateSP( PropertyEditor, &FPropertyEditor::MakeNewBlueprint )  );
 			break;
 
+		case EPropertyButton::EditConfigHierarchy:
+			NewButton = PropertyCustomizationHelpers::MakeEditConfigHierarchyButton(FSimpleDelegate::CreateSP(PropertyEditor, &FPropertyEditor::EditConfigHierarchy));
+			break;
+
 		default:
 			checkf( 0, TEXT( "Unknown button type" ) );
 			break;
diff --git a/Engine/Source/Editor/PropertyEditor/Private/PropertyEditorHelpers.h b/Engine/Source/Editor/PropertyEditor/Private/PropertyEditorHelpers.h
--- a/Engine/Source/Editor/PropertyEditor/Private/PropertyEditorHelpers.h
+++ b/Engine/Source/Editor/PropertyEditor/Private/PropertyEditorHelpers.h
@@ -25,6 +25,7 @@ namespace EPropertyButton
 		Clear,
 		Use,
 		NewBlueprint,
+		EditConfigHierarchy,
 	};
 }
 
diff --git a/Engine/Source/Editor/PropertyEditor/PropertyEditor.Build.cs b/Engine/Source/Editor/PropertyEditor/PropertyEditor.Build.cs
--- a/Engine/Source/Editor/PropertyEditor/PropertyEditor.Build.cs
+++ b/Engine/Source/Editor/PropertyEditor/PropertyEditor.Build.cs
@@ -40,7 +40,8 @@ public class PropertyEditor : ModuleRules
                 "AssetRegistry",
                 "AssetTools",
 				"ClassViewer",
-                "ContentBrowser"
+                "ContentBrowser",
+				"ConfigEditor",
 			}
 		);
 
@@ -57,6 +58,7 @@ public class PropertyEditor : ModuleRules
 				"EditorWidgets",
 				"Documentation",
                 "RHI",
+				"ConfigEditor",
 			}
 		);
 
diff --git a/Engine/Source/Editor/PropertyEditor/Public/PropertyCustomizationHelpers.h b/Engine/Source/Editor/PropertyEditor/Public/PropertyCustomizationHelpers.h
--- a/Engine/Source/Editor/PropertyEditor/Public/PropertyCustomizationHelpers.h
+++ b/Engine/Source/Editor/PropertyEditor/Public/PropertyCustomizationHelpers.h
@@ -33,7 +33,8 @@ namespace PropertyCustomizationHelpers
 	PROPERTYEDITOR_API TSharedRef<SWidget> MakeAssetPickerWithMenu( const FAssetData& InitialObject, const bool AllowClear, const TArray<const UClass*>& AllowedClasses, const TArray<UFactory*>& NewAssetFactories, FOnShouldFilterAsset OnShouldFilterAsset, FOnAssetSelected OnSet, FSimpleDelegate OnClose );
 	PROPERTYEDITOR_API TSharedRef<SWidget> MakeActorPickerAnchorButton( FOnGetActorFilters OnGetActorFilters, FOnActorSelected OnActorSelectedFromPicker );
 	PROPERTYEDITOR_API TSharedRef<SWidget> MakeActorPickerWithMenu( AActor* const InitialActor, const bool AllowClear, FOnShouldFilterActor ActorFilter, FOnActorSelected OnSet, FSimpleDelegate OnClose, FSimpleDelegate OnUseSelected );
-	PROPERTYEDITOR_API TSharedRef<SWidget> MakeInteractiveActorPicker( FOnGetAllowedClasses OnGetAllowedClasses, FOnShouldFilterActor OnShouldFilterActor, FOnActorSelected OnActorSelectedFromPicker );
+	PROPERTYEDITOR_API TSharedRef<SWidget> MakeInteractiveActorPicker(FOnGetAllowedClasses OnGetAllowedClasses, FOnShouldFilterActor OnShouldFilterActor, FOnActorSelected OnActorSelectedFromPicker);
+	PROPERTYEDITOR_API TSharedRef<SWidget> MakeEditConfigHierarchyButton(FSimpleDelegate OnEditConfigClicked, TAttribute<FText> OptionalToolTipText = FText(), TAttribute<bool> IsEnabled = true);
 
 	/** @return the UBoolProperty edit condition property if one exists. */
 	PROPERTYEDITOR_API class UBoolProperty* GetEditConditionProperty(const class UProperty* InProperty, bool& bNegate);
diff --git a/Engine/Source/Runtime/CoreUObject/Private/UObject/Obj.cpp b/Engine/Source/Runtime/CoreUObject/Private/UObject/Obj.cpp
--- a/Engine/Source/Runtime/CoreUObject/Private/UObject/Obj.cpp
+++ b/Engine/Source/Runtime/CoreUObject/Private/UObject/Obj.cpp
@@ -1868,9 +1868,9 @@ void UObject::UpdateSingleSectionOfConfigFile(const FString& ConfigIniName)
 	GConfig->LoadGlobalIniFile(FinalIniFileName, *GetClass()->ClassConfigName.ToString(), NULL, true);
 }
 
-void UObject::UpdateDefaultConfigFile()
+void UObject::UpdateDefaultConfigFile(const FString& SpecificFileLocation)
 {
-	UpdateSingleSectionOfConfigFile(GetDefaultConfigFilename());
+	UpdateSingleSectionOfConfigFile(SpecificFileLocation.IsEmpty() ? GetDefaultConfigFilename() : SpecificFileLocation);
 }
 
 void UObject::UpdateGlobalUserConfigFile()
diff --git a/Engine/Source/Runtime/CoreUObject/Public/UObject/ObjectBase.h b/Engine/Source/Runtime/CoreUObject/Public/UObject/ObjectBase.h
--- a/Engine/Source/Runtime/CoreUObject/Public/UObject/ObjectBase.h
+++ b/Engine/Source/Runtime/CoreUObject/Public/UObject/ObjectBase.h
@@ -1050,6 +1050,9 @@ namespace UM
 
 		/// [PropertyMetadata] Used for float and integer properties.  Specifies the highest that the value slider should represent.
 		UIMax,
+
+		/// [PropertyMetadata] Property is serialized to config and we should be able to set it anywhere along the config hierarchy.
+		ConfigHierarchyEditable,
 	};
 
 	// Metadata usable in UPROPERTY for customizing the behavior of Persona and UMG
diff --git a/Engine/Source/Runtime/CoreUObject/Public/UObject/UObject.h b/Engine/Source/Runtime/CoreUObject/Public/UObject/UObject.h
--- a/Engine/Source/Runtime/CoreUObject/Public/UObject/UObject.h
+++ b/Engine/Source/Runtime/CoreUObject/Public/UObject/UObject.h
@@ -779,7 +779,7 @@ public:
 	/**
 	 * Saves just the section(s) for this class into the default ini file for the class (with just the changes from base)
 	 */
-	void UpdateDefaultConfigFile();
+	void UpdateDefaultConfigFile(const FString& SpecificFileLocation = "");
 
 	/**
 	 * Saves just the section(s) for this class into the global user ini file for the class (with just the changes from base)
