Fixed a crash when opening SConfigEditor where a PlatformIni config cache has no branch.
Switched to using the existing ConfigCacheIni for platforms instead of loading a new cache from disk.

diff --git a/Engine/Source/Editor/ConfigEditor/Private/SConfigEditor.cpp b/Engine/Source/Editor/ConfigEditor/Private/SConfigEditor.cpp
--- a/Engine/Source/Editor/ConfigEditor/Private/SConfigEditor.cpp
+++ b/Engine/Source/Editor/ConfigEditor/Private/SConfigEditor.cpp
@@ -40,8 +40,6 @@ void SConfigEditor::Construct(const FArguments& InArgs, TWeakFieldPtr<FProperty>
 
 	EditProperty = InEditProperty;
 
-	LocalConfigCache = MakeShareable(new FConfigCacheIni(EConfigCacheType::Temporary));
-
 	// initialize details view
 	FDetailsViewArgs DetailsViewArgs;
 	{
@@ -96,15 +94,18 @@ void SConfigEditor::CreateDisplayObjectForSelectedTargetPlatform()
 	PropHelper->EditProperty = EditProperty.Get();
 	FString ClassConfigName = PropHelper->EditProperty->GetOwnerClass()->ClassConfigName.ToString();
 
-	FConfigFile PlatformIniFile;
-	LocalConfigCache->LoadLocalIniFile(PlatformIniFile, *ClassConfigName, true, *SelectedTargetPlatform);
-
-	for (const auto& IniFile : PlatformIniFile.Branch->Hierarchy)
+	if (FConfigCacheIni* PlatformConfigCache = FConfigCacheIni::ForPlatform(*SelectedTargetPlatform))
 	{
-		UPropertyConfigFileDisplayRow* ConfigFilePropertyObj = NewObject<UPropertyConfigFileDisplayRow>(GetTransientPackage(), *IniFile.Value);
-		ConfigFilePropertyObj->InitWithConfigAndProperty(IniFile.Value, PropHelper->EditProperty.Get());
-
-		PropHelper->ConfigFilePropertyObjects.Add(ConfigFilePropertyObj);
+		if (FConfigBranch* Branch = PlatformConfigCache->FindBranch(*ClassConfigName, ClassConfigName))
+		{
+			for (const auto& IniFile : Branch->Hierarchy)
+			{
+				UPropertyConfigFileDisplayRow* ConfigFilePropertyObj = NewObject<UPropertyConfigFileDisplayRow>(GetTransientPackage(), *IniFile.Value);
+				ConfigFilePropertyObj->InitWithConfigAndProperty(IniFile.Value, PropHelper->EditProperty.Get());
+
+				PropHelper->ConfigFilePropertyObjects.Add(ConfigFilePropertyObj);
+			}
+		}
 	}
 
 	DetailsView->SetObject(PropHelper.Get());
diff --git a/Engine/Source/Editor/ConfigEditor/Private/SConfigEditor.h b/Engine/Source/Editor/ConfigEditor/Private/SConfigEditor.h
--- a/Engine/Source/Editor/ConfigEditor/Private/SConfigEditor.h
+++ b/Engine/Source/Editor/ConfigEditor/Private/SConfigEditor.h
@@ -52,9 +52,6 @@ private:
 	// The display area for the config hierarchy editor.
 	TSharedPtr<SWidget> PropertyValueEditor;
 
-	// We keep a cache of config files for the provided platforms.
-	TSharedPtr<FConfigCacheIni> LocalConfigCache;
-
 	// The bulk of the display of this hierarchy.
 	TSharedPtr<IDetailsView> DetailsView;
 
diff --git a/Engine/Source/Editor/DetailCustomizations/Private/ConfigEditorPropertyDetails.cpp b/Engine/Source/Editor/DetailCustomizations/Private/ConfigEditorPropertyDetails.cpp
--- a/Engine/Source/Editor/DetailCustomizations/Private/ConfigEditorPropertyDetails.cpp
+++ b/Engine/Source/Editor/DetailCustomizations/Private/ConfigEditorPropertyDetails.cpp
@@ -124,7 +124,14 @@ void FConfigPropertyHelperDetails::OnPropertyValueChanged(UObject* Object, FProp
 			NewFile.UpdateSinglePropertyInSection(*ConfigIniName, *PropertyName, *SectionName);
 
 			// reload the file, so that it refresh the cache internally.
-			FConfigContext::ForceReloadIntoGConfig().Load(*OriginalProperty->GetOwnerClass()->ClassConfigName.ToString());
+			if (GConfig->FindConfigFile(ConfigIniName))
+			{
+				FConfigContext::ForceReloadIntoGConfig().Load(*OriginalProperty->GetOwnerClass()->ClassConfigName.ToString());
+			}
+			else
+			{
+				FConfigCacheIni::ClearOtherPlatformConfigs();
+			}
 
 			// Update the CDO, as this change might have had an impact on it's value.
 			OriginalProperty->GetOwnerClass()->GetDefaultObject()->ReloadConfig();
